Fix for uninitialized structure members

diff --git a/Engine/Source/Runtime/EngineMessages/Public/EngineServiceMessages.h b/Engine/Source/Runtime/EngineMessages/Public/EngineServiceMessages.h
--- a/Engine/Source/Runtime/EngineMessages/Public/EngineServiceMessages.h
+++ b/Engine/Source/Runtime/EngineMessages/Public/EngineServiceMessages.h
@@ -35,11 +35,11 @@ struct FEngineServicePong
 
 	/** Holds the engine version. */
 	UPROPERTY(EditAnywhere, Category="Message")
-	int32 EngineVersion;
+	int32 EngineVersion=0;
 
 	/** Holds a flag indicating whether game play has begun. */
 	UPROPERTY(EditAnywhere, Category="Message")
-	bool HasBegunPlay;
+	bool HasBegunPlay=false;
 
 	/** Holds the instance identifier. */
 	UPROPERTY(EditAnywhere, Category="Message")
@@ -55,7 +55,7 @@ struct FEngineServicePong
 
 	/** Holds the time in seconds since the world was loaded. */
 	UPROPERTY(EditAnywhere, Category="Message")
-	float WorldTimeSeconds;
+	float WorldTimeSeconds=0.0f;
 };
 
 
