Fix missing default value error.

diff --git a/Engine/Source/Runtime/EngineMessages/Public/TraceControlMessages.h b/Engine/Source/Runtime/EngineMessages/Public/TraceControlMessages.h
--- a/Engine/Source/Runtime/EngineMessages/Public/TraceControlMessages.h
+++ b/Engine/Source/Runtime/EngineMessages/Public/TraceControlMessages.h
@@ -39,22 +39,22 @@ struct FTraceControlStatus
 	FGuid TraceGuid;
 
 	UPROPERTY(EditAnywhere, Category="Message")
-	uint64 BytesSent;
+	uint64 BytesSent = 0;
 
 	UPROPERTY(EditAnywhere, Category="Message")
-	uint64 BytesTraced;
+	uint64 BytesTraced = 0;
 
 	UPROPERTY(EditAnywhere, Category="Message")
-	uint64 MemoryUsed;
+	uint64 MemoryUsed = 0;
 
 	UPROPERTY(EditAnywhere, Category="Message")
-	uint32 CacheAllocated; 
+	uint32 CacheAllocated = 0; 
 
 	UPROPERTY(EditAnywhere, Category="Message")
-	uint32 CacheUsed;
+	uint32 CacheUsed = 0;
 
 	UPROPERTY(EditAnywhere, Category="Message")
-	uint32 CacheWaste;
+	uint32 CacheWaste = 0;
 };
 
 /**
@@ -76,16 +76,16 @@ struct FTraceControlSettings
 	GENERATED_USTRUCT_BODY()
 
 	UPROPERTY(EditAnywhere, Category="Message")
-	bool bUseWorkerThread;
+	bool bUseWorkerThread = false;
 	
 	UPROPERTY(EditAnywhere, Category="Message")
-	bool bUseImportantCache;
+	bool bUseImportantCache = false;
 
 	UPROPERTY(EditAnywhere, Category="Message")
-	uint32 TailSizeBytes;
+	uint32 TailSizeBytes = 0;
 	
 	UPROPERTY(EditAnywhere, Category="Message")
-	bool bStatNamedEvents;
+	bool bStatNamedEvents = false;
 };
 
 /**
@@ -97,7 +97,7 @@ struct FTraceControlChannelsPing
 	GENERATED_USTRUCT_BODY()
 	
 	UPROPERTY(EditAnywhere, Category="Message")
-	uint32 KnownChannelCount;
+	uint32 KnownChannelCount = 0;
 };
 
 /**
@@ -205,7 +205,7 @@ struct FTraceControlStartCommon
 	FString Channels;
 
 	UPROPERTY(EditAnywhere, Category="Message")
-	bool bExcludeTail;
+	bool bExcludeTail = false;
 };
 
 /**
@@ -232,7 +232,7 @@ struct FTraceControlFile : public FTraceControlStartCommon
 	FString File;
 	
 	UPROPERTY(EditAnywhere, Category="Message")
-	bool bTruncateFile;
+	bool bTruncateFile = false;
 };
 
 /**
@@ -301,5 +301,5 @@ struct FTraceControlScreenshot
 	FString Name;
 	
 	UPROPERTY(EditAnywhere, Category="Message")
-	bool bShowUI;
+	bool bShowUI = true;
 };
