Add missing initialization

diff --git a/Engine/Source/Runtime/EngineMessages/Public/TraceControlMessages.h b/Engine/Source/Runtime/EngineMessages/Public/TraceControlMessages.h
--- a/Engine/Source/Runtime/EngineMessages/Public/TraceControlMessages.h
+++ b/Engine/Source/Runtime/EngineMessages/Public/TraceControlMessages.h
@@ -96,7 +96,7 @@ struct FTraceChannelPreset
 	FString ChannelList;
 
 	UPROPERTY(EditAnywhere, Category = "Message")
-	bool bIsReadOnly;
+	bool bIsReadOnly = false;
 };
 
 /**
