[Insights] Fix default values for messages.

#Synced-CL 49918891

diff --git a/Engine/Source/Runtime/EngineMessages/Public/TraceControlMessages.h b/Engine/Source/Runtime/EngineMessages/Public/TraceControlMessages.h
--- a/Engine/Source/Runtime/EngineMessages/Public/TraceControlMessages.h
+++ b/Engine/Source/Runtime/EngineMessages/Public/TraceControlMessages.h
@@ -329,10 +329,10 @@ struct FTraceControlSecureResponse
 	FString Host;
 
 	UPROPERTY(EditAnywhere, Category="Message")
-	uint16 Port;
+	uint16 Port = 0;
 
 	UPROPERTY(EditAnywhere, Category="Message")
-	ETraceSecureMode Mode;
+	ETraceSecureMode Mode = ETraceSecureMode::SelfSignedCertAndAuthToken;
 
 	UPROPERTY(EditAnywhere, Category="Message")
 	TArray<uint8> AuthorizationToken;
